- Moved platform natvis and install scripts to PlatformExtensions

diff --git a/Engine/Extras/VisualStudioDebugging/InstallVisualizers.bat b/Engine/Extras/VisualStudioDebugging/InstallVisualizers.bat
--- a/Engine/Extras/VisualStudioDebugging/InstallVisualizers.bat
+++ b/Engine/Extras/VisualStudioDebugging/InstallVisualizers.bat
@@ -26,8 +26,10 @@ if "%VS110COMNTOOLS%" neq "" (
   echo Installed visualizers for Visual Studio 2012
 )
 
-if exist PS4\InstallPS4Visualizer.bat (
-  call PS4\InstallPS4Visualizer.bat
+for /d %%i in (..\..\Platforms\*) do (
+	if exist "%%i\Extras\VisualStudioDebugging\InstallVisualizer.bat" (
+		call "%%i\Extras\VisualStudioDebugging\InstallVisualizer.bat"
+	)
 )
 
 if "%FoundVSInstance%" equ "0" (
