Renamed UE4.natvis to Unreal.natvis
Fixed up Project File Generator within the Unreal Build Tool

diff --git a/Engine/Build/InstalledEngineFilters.xml b/Engine/Build/InstalledEngineFilters.xml
--- a/Engine/Build/InstalledEngineFilters.xml
+++ b/Engine/Build/InstalledEngineFilters.xml
@@ -282,7 +282,7 @@
 
 		<!-- Extra tools -->
 		Engine/Extras/3dsMaxScripts/...
-		Engine/Extras/VisualStudioDebugging/UE4.natvis
+		Engine/Extras/VisualStudioDebugging/Unreal.natvis
 		Engine/Extras/MayaLiveLink/...
 		Engine/Extras/MayaVelocityGridExporter/...
 		Engine/Extras/UnrealVS/...
diff --git a/Engine/Extras/VisualStudioDebugging/InstallVisualizers.bat b/Engine/Extras/VisualStudioDebugging/InstallVisualizers.bat
--- a/Engine/Extras/VisualStudioDebugging/InstallVisualizers.bat
+++ b/Engine/Extras/VisualStudioDebugging/InstallVisualizers.bat
@@ -7,22 +7,22 @@ pushd %~dp0
 if "%VS140COMNTOOLS%" neq "" (
   set FoundVSInstance=1
   mkdir "%USERPROFILE%\Documents\Visual Studio 2015\Visualizers"
-  copy UE4.natvis "%USERPROFILE%\Documents\Visual Studio 2015\Visualizers"
-  copy UE4_Android_Nsight.dat "%VS140COMNTOOLS%\..\IDE\Extensions\NVIDIA\Nsight Tegra\3.4\Debuggers\Visualizers"
+  copy Unreal.natvis "%USERPROFILE%\Documents\Visual Studio 2015\Visualizers"
+  copy Unreal_Android_Nsight.dat "%VS140COMNTOOLS%\..\IDE\Extensions\NVIDIA\Nsight Tegra\3.4\Debuggers\Visualizers"
   echo Installed visualizers for Visual Studio 2015
 )
 
 if "%VS120COMNTOOLS%" neq "" (
   set FoundVSInstance=1
   mkdir "%USERPROFILE%\Documents\Visual Studio 2013\Visualizers"
-  copy UE4.natvis "%USERPROFILE%\Documents\Visual Studio 2013\Visualizers"
+  copy Unreal.natvis "%USERPROFILE%\Documents\Visual Studio 2013\Visualizers"
   echo Installed visualizers for Visual Studio 2013
 )
 
 if "%VS110COMNTOOLS%" neq "" (
   set FoundVSInstance=1
   mkdir "%USERPROFILE%\Documents\Visual Studio 2012\Visualizers"
-  copy UE4.natvis "%USERPROFILE%\Documents\Visual Studio 2012\Visualizers"
+  copy Unreal.natvis "%USERPROFILE%\Documents\Visual Studio 2012\Visualizers"
   echo Installed visualizers for Visual Studio 2012
 )
 
diff --git a/Engine/Source/Programs/UnrealBuildTool/ProjectFiles/VisualStudio/VCProjectFileGenerator.cs b/Engine/Source/Programs/UnrealBuildTool/ProjectFiles/VisualStudio/VCProjectFileGenerator.cs
--- a/Engine/Source/Programs/UnrealBuildTool/ProjectFiles/VisualStudio/VCProjectFileGenerator.cs
+++ b/Engine/Source/Programs/UnrealBuildTool/ProjectFiles/VisualStudio/VCProjectFileGenerator.cs
@@ -630,7 +630,7 @@ namespace UnrealBuildTool
 				}
 
 				// Get the path to the visualizers file. Try to make it relative to the solution directory, but fall back to a full path if it's a foreign project.
-				FileReference VisualizersFile = FileReference.Combine(UnrealBuildTool.EngineDirectory, "Extras", "VisualStudioDebugging", "UE4.natvis");
+				FileReference VisualizersFile = FileReference.Combine(UnrealBuildTool.EngineDirectory, "Extras", "VisualStudioDebugging", "Unreal.natvis");
 
 				// Add the visualizers at the solution level. Doesn't seem to be picked up from a makefile project in VS2017 15.8.5.
 				VCSolutionFileContent.AppendLine(String.Format("Project(\"{0}\") = \"Visualizers\", \"Visualizers\", \"{{1CCEC849-CC72-4C59-8C36-2F7C38706D4C}}\"", SolutionFolderEntryGUID));
diff --git a/Engine/Source/Programs/UnrealBuildTool/ProjectFiles/VisualStudioCode/VSCodeProjectFileGenerator.cs b/Engine/Source/Programs/UnrealBuildTool/ProjectFiles/VisualStudioCode/VSCodeProjectFileGenerator.cs
--- a/Engine/Source/Programs/UnrealBuildTool/ProjectFiles/VisualStudioCode/VSCodeProjectFileGenerator.cs
+++ b/Engine/Source/Programs/UnrealBuildTool/ProjectFiles/VisualStudioCode/VSCodeProjectFileGenerator.cs
@@ -949,7 +949,7 @@ namespace UnrealBuildTool
 								OutFile.AddField("externalConsole", true);
 
 								OutFile.AddField("type", "cppvsdbg");
-								OutFile.AddField("visualizerFile", MakeUnquotedPathString(FileReference.Combine(UE4ProjectRoot, "Engine", "Extras", "VisualStudioDebugging", "UE4.natvis"), EPathType.Absolute));
+								OutFile.AddField("visualizerFile", MakeUnquotedPathString(FileReference.Combine(UE4ProjectRoot, "Engine", "Extras", "VisualStudioDebugging", "Unreal.natvis"), EPathType.Absolute));
 							}
 							else
 							{
