VisualStudio: Add a .natstepfilter for common Unreal parameter constructor functions to avoid stepping into the parameter construction when trying to step into a function that takes the parameter.
Remove InstallVisualizers.bat; we no longer support VS2015 and earlier, and versions after 2015 use .natvis that is referenced from the solution rather than requiring a system-wide install.
#rn Minor Core

diff --git a/Engine/Extras/VisualStudioDebugging/InstallStepFilters.bat b/Engine/Extras/VisualStudioDebugging/InstallStepFilters.bat
new file mode 100755
--- /dev/null
+++ b/Engine/Extras/VisualStudioDebugging/InstallStepFilters.bat
@@ -0,0 +1,54 @@
+@setlocal
+@echo off
+:: Change to directory where batch file is located.  We'll restore this later with "popd"
+pushd %~dp0
+
+
+:: Main Script
+SET FOUND_VERSIONS=
+set VSWHERE="%ProgramFiles(x86)%\Microsoft Visual Studio\Installer\vswhere.exe"
+
+if exist %VSWHERE% (
+    for /F "tokens=1,2" %%i in ('%VSWHERE% -all') do (
+        if "%%i"=="catalog_productLineVersion:" (
+          call :InstallVersion %%j
+        )
+    )
+)
+
+if "%FOUND_VERSIONS%" equ "" (
+  echo ERROR: Could not find a valid version of Visual Studio installed
+)
+goto :End
+
+
+
+:: function InstallVersion // InstallVersion <VersionName>
+:InstallVersion
+
+set VERSION_NAME=%1
+:: Strip leading whitespace from version name
+for /f "tokens=* delims= " %%a in ("%VERSION_NAME%") do set VERSION_NAME=%%a
+if "%VERSION_NAME%"=="" exit /b 0
+
+set ALREADY_EXISTS=0
+for /F %%i in ('echo. %FOUND_VERSIONS% ^| find "_SEP_%VERSION_NAME%_SEP_"') do set ALREADY_EXISTS=1
+if %ALREADY_EXISTS%==1 exit /b 0
+set FOUND_VERSIONS=%FOUND_VERSIONS%_SEP_%VERSION_NAME%_SEP_
+
+set VISUALIZER_DIR="%USERPROFILE%\Documents\Visual Studio %VERSION_NAME%\Visualizers"
+if not exist %VISUALIZER_DIR% mkdir %VISUALIZER_DIR%
+
+echo Installing visualizers for Visual Studio %VERSION_NAME%...
+echo     copy Unreal.natstepfilter %VISUALIZER_DIR%
+copy     Unreal.natstepfilter %VISUALIZER_DIR% 1>nul
+
+exit /b 0
+
+
+:: goto-label :End // restores settings and exits
+:End
+popd
+pause
+endlocal
+exit /b
diff --git a/Engine/Extras/VisualStudioDebugging/InstallVisualizers.bat b/Engine/Extras/VisualStudioDebugging/InstallVisualizers.bat
deleted file mode 100755
