Copy Unreal.natjmc as part of InstallStepFilters.bat
- This is generally unnecessary because the file will get used if it's part of the Visual Studio solution, which it is. But this ensures the batch file works as expected.

#Synced-CL 48050961

diff --git a/Engine/Extras/VisualStudioDebugging/InstallStepFilters.bat b/Engine/Extras/VisualStudioDebugging/InstallStepFilters.bat
--- a/Engine/Extras/VisualStudioDebugging/InstallStepFilters.bat
+++ b/Engine/Extras/VisualStudioDebugging/InstallStepFilters.bat
@@ -39,10 +39,14 @@ set FOUND_VERSIONS=%FOUND_VERSIONS%_SEP_%VERSION_NAME%_SEP_
 set VISUALIZER_DIR="%USERPROFILE%\Documents\Visual Studio %VERSION_NAME%\Visualizers"
 if not exist %VISUALIZER_DIR% mkdir %VISUALIZER_DIR%
 
-echo Installing visualizers for Visual Studio %VERSION_NAME%...
+echo Installing debugger step filters for Visual Studio %VERSION_NAME%...
+
 echo     copy Unreal.natstepfilter %VISUALIZER_DIR%
 copy     Unreal.natstepfilter %VISUALIZER_DIR% 1>nul
 
+echo     copy Unreal.natjmc %VISUALIZER_DIR%
+copy     Unreal.natjmc %VISUALIZER_DIR% 1>nul
+
 exit /b 0
 
 
